TexureFormats: delay load DLLs for ISPC and ETC2 which should be rarely used now

diff --git a/Engine/Source/Developer/TextureFormatETC2/Private/TextureFormatETC2.cpp b/Engine/Source/Developer/TextureFormatETC2/Private/TextureFormatETC2.cpp
--- a/Engine/Source/Developer/TextureFormatETC2/Private/TextureFormatETC2.cpp
+++ b/Engine/Source/Developer/TextureFormatETC2/Private/TextureFormatETC2.cpp
@@ -127,7 +127,7 @@ static bool CompressImageUsingQonvert(
 
 	if (Qonvert(&SrcImg, &DstImg) != Q_SUCCESS)
 	{
-		UE_LOG(LogTextureFormatETC2, Fatal, TEXT("CONVERSION FAILED"));
+		UE_LOG(LogTextureFormatETC2, Error, TEXT("Qonvert failed"));
 		return false;
 	}
 
@@ -140,6 +140,54 @@ static bool CompressImageUsingQonvert(
  */
 class FTextureFormatETC2 : public ITextureFormat
 {
+public:
+
+#if PLATFORM_WINDOWS
+	void*	TextureConverterHandle = NULL;
+	FString AppLocalBinariesRoot = FPaths::EngineDir() / TEXT("Binaries/ThirdParty/AppLocalDependencies/Win64/Microsoft.VC.CRT");
+	FString QualCommBinariesRoot = FPaths::EngineDir() / TEXT("Binaries/ThirdParty/QualComm/Win64");
+	FString QualCommBinaryName = TEXT("TextureConverter.dll");
+#endif
+
+	FTextureFormatETC2()
+	{
+	}
+
+	~FTextureFormatETC2()
+	{
+#if PLATFORM_WINDOWS
+		if (TextureConverterHandle)
+		{
+			FPlatformProcess::FreeDllHandle(TextureConverterHandle);
+			TextureConverterHandle = nullptr;
+		}
+#endif
+	}
+
+	void LoadDLLOnce()
+	{
+		// Load DLL on first use, not at startup
+
+#if PLATFORM_WINDOWS
+		// should only be called once
+		check(TextureConverterHandle == nullptr);
+
+		UE_LOG(LogTextureFormatETC2, Display, TEXT("ETC2 Texture loading DLL: %s"), *QualCommBinaryName);
+
+		FPlatformProcess::PushDllDirectory(*AppLocalBinariesRoot);
+		void* handle = FPlatformProcess::GetDllHandle(*(QualCommBinariesRoot / QualCommBinaryName));
+		FPlatformProcess::PopDllDirectory(*AppLocalBinariesRoot);
+		
+		TextureConverterHandle = handle;
+
+		if (handle == nullptr )
+		{
+			UE_LOG(LogTextureFormatETC2, Error, TEXT("ETC2 Texture %s requested but could not be loaded"), *QualCommBinaryName);
+		}
+#endif
+	}
+
+
 	virtual bool AllowParallelBuild() const override
 	{
 		return !PLATFORM_MAC; // On Mac Qualcomm's TextureConverter library is not thead-safe
@@ -161,10 +209,7 @@ class FTextureFormatETC2 : public ITextureFormat
 
 	virtual void GetSupportedFormats(TArray<FName>& OutFormats) const override
 	{
-		for (int32 i = 0; i < UE_ARRAY_COUNT(GSupportedTextureFormatNames); ++i)
-		{
-			OutFormats.Add(GSupportedTextureFormatNames[i]);
-		}
+		OutFormats.Append(GSupportedTextureFormatNames, UE_ARRAY_COUNT(GSupportedTextureFormatNames));
 	}
 
 	virtual FTextureFormatCompressorCaps GetFormatCapabilities() const override
@@ -203,6 +248,9 @@ class FTextureFormatETC2 : public ITextureFormat
 		FCompressedImage2D& OutCompressedImage
 		) const override
 	{
+
+		UE_CALL_ONCE( const_cast<FTextureFormatETC2 *>(this)->LoadDLLOnce );
+
 		FImage Image;
 		InImage.CopyTo(Image, ERawImageFormat::BGRA8, BuildSettings.GetDestGammaSpace());
 
@@ -235,35 +283,22 @@ class FTextureFormatETC2 : public ITextureFormat
 	}
 };
 
-static ITextureFormat* Singleton = NULL;
 
 
 
-#if PLATFORM_WINDOWS
-	void*	TextureConverterHandle = NULL;
-	FString AppLocalBinariesRoot = FPaths::EngineDir() / TEXT("Binaries/ThirdParty/AppLocalDependencies/Win64/Microsoft.VC.CRT");
-	FString QualCommBinariesRoot = FPaths::EngineDir() / TEXT("Binaries/ThirdParty/QualComm/Win64");
-	FString QualCommBinaryName = TEXT("TextureConverter.dll");
-#endif
-
-
 class FTextureFormatETC2Module : public ITextureFormatModule
 {
 public:
+	ITextureFormat* Singleton = NULL;
+
 	FTextureFormatETC2Module() { }
 	virtual ~FTextureFormatETC2Module()
 	{
-		ITextureFormat* p = Singleton;
-		Singleton = nullptr;
-		if (p)
-			delete p;
-
-#if PLATFORM_WINDOWS
-		void* handle = TextureConverterHandle;
-		TextureConverterHandle = nullptr;
-		if (handle)
-			FPlatformProcess::FreeDllHandle(handle);
-#endif
+		if ( Singleton )
+		{
+			delete Singleton;
+			Singleton = nullptr;
+		}
 	}
 
 	virtual void StartupModule() override
@@ -274,30 +309,6 @@ public:
 
 	virtual ITextureFormat* GetTextureFormat()
 	{
-#if PLATFORM_WINDOWS
-		// @@!! TODO: move this to first use, not at startup
-
-		if (TextureConverterHandle == nullptr)
-		{
-			UE_LOG(LogTextureFormatETC2, Display, TEXT("ETC2 Texture loading DLL: %s"), *QualCommBinaryName);
-			FPlatformProcess::PushDllDirectory(*AppLocalBinariesRoot);
-			void* handle = FPlatformProcess::GetDllHandle(*(QualCommBinariesRoot / QualCommBinaryName));
-			FPlatformProcess::PopDllDirectory(*AppLocalBinariesRoot);
-			if (handle == nullptr )
-			{
-				UE_LOG(LogTextureFormatETC2, Warning, TEXT("ETC2 Texture %s requested but could not be loaded"), *QualCommBinaryName);
-				return nullptr;
-			}
-			TextureConverterHandle = handle;
-		}
-
-		if (TextureConverterHandle == nullptr)
-		{
-			UE_LOG(LogTextureFormatETC2, Warning, TEXT("ETC2 Texture %s requested but could not be loaded"), *QualCommBinaryName);
-			return nullptr;
-		}
-#endif
-
 		if ( Singleton == nullptr )  // not thread safe
 		{
 			FTextureFormatETC2* ptr = new FTextureFormatETC2();
diff --git a/Engine/Source/Developer/TextureFormatIntelISPCTexComp/Private/TextureFormatIntelISPCTexComp.cpp b/Engine/Source/Developer/TextureFormatIntelISPCTexComp/Private/TextureFormatIntelISPCTexComp.cpp
--- a/Engine/Source/Developer/TextureFormatIntelISPCTexComp/Private/TextureFormatIntelISPCTexComp.cpp
+++ b/Engine/Source/Developer/TextureFormatIntelISPCTexComp/Private/TextureFormatIntelISPCTexComp.cpp
@@ -560,12 +560,52 @@ static void IntelASTCCompressScans(FASTCEncoderSettings* pEncSettings, FImage* p
 class FTextureFormatIntelISPCTexComp : public ITextureFormat
 {
 public:
+
+	void* mDllHandle = nullptr;
+
 	FTextureFormatIntelISPCTexComp()
 	{
 	}
 	virtual ~FTextureFormatIntelISPCTexComp()
 	{
+		if ( mDllHandle != nullptr )
+		{
+			FPlatformProcess::FreeDllHandle(mDllHandle);
+			mDllHandle = nullptr;
+		}
+	}
+
+	void LoadDLLOnce()
+	{
+		// LoadDLL is delayed to first use
+		// should only be called once
+		check( mDllHandle == nullptr );
+
+		FString DLLPath;
+#if PLATFORM_WINDOWS
+#if PLATFORM_64BITS
+		DLLPath = FPaths::EngineDir() / TEXT("Binaries/ThirdParty/Intel/ISPCTexComp/Win64-Release/ispc_texcomp.dll");
+#else	//32-bit platform
+		DLLPath = FPaths::EngineDir() / TEXT("Binaries/ThirdParty/Intel/ISPCTexComp/Win32-Release/ispc_texcomp.dll");
+#endif
+#elif PLATFORM_MAC
+		DLLPath = TEXT("libispc_texcomp.dylib");
+#elif PLATFORM_LINUX
+		DLLPath = FPaths::EngineDir() / TEXT("Binaries/ThirdParty/Intel/ISPCTexComp/Linux64-Release/libispc_texcomp.so");
+#endif
+
+		if (DLLPath.Len() > 0)
+		{
+			mDllHandle = FPlatformProcess::GetDllHandle(*DLLPath);
+			UE_CLOG(mDllHandle == nullptr, LogTextureFormatIntelISPCTexComp, Error, TEXT("Unable to load %s"), *DLLPath);
+		}
+		else
+		{
+			UE_LOG(LogTextureFormatIntelISPCTexComp, Error, TEXT("Platform does not have an ispc_texcomp DLL/library"));
+		}
 	}
+
+
 	virtual bool AllowParallelBuild() const override
 	{
 		return true;
@@ -596,10 +636,7 @@ public:
 
 	virtual void GetSupportedFormats(TArray<FName>& OutFormats) const override
 	{
-		for (int32 i = 0; i < UE_ARRAY_COUNT(GSupportedTextureFormatNames); ++i)
-		{
-			OutFormats.Add(GSupportedTextureFormatNames[i]);
-		}
+		OutFormats.Append(GSupportedTextureFormatNames, UE_ARRAY_COUNT(GSupportedTextureFormatNames));
 	}
 
 	virtual FTextureFormatCompressorCaps GetFormatCapabilities() const override
@@ -642,7 +679,6 @@ public:
 		// ISPCtexcomp crashes on images not of block-aligned size
 
 		// Allocate temp buffer
-		//@TODO: Optimize away this temp buffer (could avoid last FMemory::Memcpy)
 		TArray64<uint8> TempBuffer;
 		TempBuffer.SetNumUninitialized(AlignedTotalSize);
 
@@ -744,6 +780,9 @@ public:
 		check(InImage.SizeY > 0);
 		check(InImage.NumSlices > 0);
 
+		// Load DLL on first use :
+		UE_CALL_ONCE( const_cast<FTextureFormatIntelISPCTexComp *>(this)->LoadDLLOnce );
+
 		bool bCompressionSucceeded = false;
 
 		int BlockWidth = 0;
@@ -911,18 +950,14 @@ class FTextureFormatIntelISPCTexCompModule : public ITextureFormatModule
 public:
 	FTextureFormatIntelISPCTexCompModule()
 	{
-		mDllHandle = nullptr;
 	}
 
 	virtual ~FTextureFormatIntelISPCTexCompModule()
 	{
-		delete Singleton;
-		Singleton = nullptr;
-
-		if ( mDllHandle != nullptr )
+		if ( Singleton )
 		{
-			FPlatformProcess::FreeDllHandle(mDllHandle);
-			mDllHandle = nullptr;
+			delete Singleton;
+			Singleton = nullptr;
 		}
 	}
 	
@@ -932,38 +967,12 @@ public:
 	{
 		if (!Singleton)
 		{
-			// @@!! todo: move this to first use, not at startup
-
-			FString DLLPath;
-#if PLATFORM_WINDOWS
-	#if PLATFORM_64BITS
-			DLLPath = FPaths::EngineDir() / TEXT("Binaries/ThirdParty/Intel/ISPCTexComp/Win64-Release/ispc_texcomp.dll");
-	#else	//32-bit platform
-			DLLPath = FPaths::EngineDir() / TEXT("Binaries/ThirdParty/Intel/ISPCTexComp/Win32-Release/ispc_texcomp.dll");
-	#endif
-#elif PLATFORM_MAC
-			DLLPath = TEXT("libispc_texcomp.dylib");
-#elif PLATFORM_LINUX
-			DLLPath = FPaths::EngineDir() / TEXT("Binaries/ThirdParty/Intel/ISPCTexComp/Linux64-Release/libispc_texcomp.so");
-#endif
-
-			if (DLLPath.Len() > 0)
-			{
-				mDllHandle = FPlatformProcess::GetDllHandle(*DLLPath);
-				UE_CLOG(mDllHandle == nullptr, LogTextureFormatIntelISPCTexComp, Warning, TEXT("Unable to load %s"), *DLLPath);
-			}
-			else
-			{
-				UE_LOG(LogTextureFormatIntelISPCTexComp, Warning, TEXT("Platform does not have an ispc_texcomp DLL/library"));
-			}
-
 			Singleton = new FTextureFormatIntelISPCTexComp();
 		}
 		return Singleton;
 	}
 
 	static inline UE::DerivedData::TBuildFunctionFactory<FIntelISPCTexCompTextureBuildFunction> BuildFunctionFactory;
-	void* mDllHandle;
 };
 
 IMPLEMENT_MODULE(FTextureFormatIntelISPCTexCompModule, TextureFormatIntelISPCTexComp);
