Workaround for: error LNK2019: unresolved external symbol __imp___std_init_once_begin_initialize
when linking EtcLib.lib with a vs2022 17.2 toolchain

diff --git a/Engine/Source/Developer/TextureFormatETC2/Private/TextureFormatETC2.cpp b/Engine/Source/Developer/TextureFormatETC2/Private/TextureFormatETC2.cpp
--- a/Engine/Source/Developer/TextureFormatETC2/Private/TextureFormatETC2.cpp
+++ b/Engine/Source/Developer/TextureFormatETC2/Private/TextureFormatETC2.cpp
@@ -24,6 +24,13 @@
 #include "EtcErrorMetric.h"
 #include "EtcImage.h"
 
+// Workaround for: error LNK2019: unresolved external symbol __imp___std_init_once_begin_initialize referenced in function "void __cdecl std::call_once
+// https://developercommunity.visualstudio.com/t/-imp-std-init-once-complete-unresolved-external-sy/1684365
+#if defined(_MSC_VER) && (_MSC_VER >= 1932)  // Visual Studio 2022 version 17.2+
+#    pragma comment(linker, "/alternatename:__imp___std_init_once_complete=__imp_InitOnceComplete")
+#    pragma comment(linker, "/alternatename:__imp___std_init_once_begin_initialize=__imp_InitOnceBeginInitialize")
+#endif
+
 DEFINE_LOG_CATEGORY_STATIC(LogTextureFormatETC2, Log, All);
 
 class FETC2TextureBuildFunction final : public FTextureBuildFunction
