TextureFormatETC2 sanitize nan/inf floats before they reach etc2comp

diff --git a/Engine/Source/Developer/TextureFormatETC2/Private/TextureFormatETC2.cpp b/Engine/Source/Developer/TextureFormatETC2/Private/TextureFormatETC2.cpp
--- a/Engine/Source/Developer/TextureFormatETC2/Private/TextureFormatETC2.cpp
+++ b/Engine/Source/Developer/TextureFormatETC2/Private/TextureFormatETC2.cpp
@@ -71,11 +71,13 @@ static FName GSupportedTextureFormatNames[] =
 
 #undef ENUM_SUPPORTED_FORMATS
 
+// note InSourceData is not const, can be mutated by sanitize
 static bool CompressImageUsingEtc2comp(
-	const void* InSourceData,
+	FLinearColor * InSourceColors,
 	EPixelFormat PixelFormat,
 	int32 SizeX,
 	int32 SizeY,
+	int64 NumPixels,
 	EGammaSpace TargetGammaSpace,
 	TArray64<uint8>& OutCompressedData)
 {
@@ -104,21 +106,24 @@ static bool CompressImageUsingEtc2comp(
 	// RGBA, REC709, NUMERIC will set RGB to 0 if all pixels in the block are transparent (A=0)
 	const Etc::ErrorMetric EtcErrorMetric = Etc::RGBX;
 	const float EtcEffort = ETCCOMP_DEFAULT_EFFORT_LEVEL;
+	// threads used by etc2comp :
 	const unsigned int MAX_JOBS = 8;
 	const unsigned int NUM_JOBS = 8;
+	// to run etc2comp synchronously :
+	//const unsigned int MAX_JOBS = 0;
+	//const unsigned int NUM_JOBS = 0;
 
 	unsigned char* paucEncodingBits = nullptr;
 	unsigned int uiEncodingBitsBytes = 0;
 	unsigned int uiExtendedWidth = 0;
 	unsigned int uiExtendedHeight = 0;
 	int iEncodingTime_ms = 0;
-	float* SourceData = (float*)InSourceData;
+	float* SourceData = &InSourceColors[0].Component(0);
 	
 	// InSourceData is a linear color, we need to feed float* data to the codec in a target color space
 	TArray64<float> IntermediateData;
 	if (TargetGammaSpace == EGammaSpace::sRGB)
 	{
-		int64 NumPixels = SizeX * SizeY;
 		IntermediateData.Reserve(NumPixels * 4);
 		IntermediateData.AddUninitialized(NumPixels * 4);
 
@@ -134,6 +139,38 @@ static bool CompressImageUsingEtc2comp(
 		
 		SourceData = IntermediateData.GetData();
 	}
+	else
+	{
+		int64 NumFloats = NumPixels * 4;
+
+		for(int64 Idx =0 ;Idx < NumFloats;Idx++)
+		{
+			// sanitize inf and nan :
+			float f = SourceData[Idx];
+			if ( f >= -FLT_MAX && f <= FLT_MAX )
+			{
+				// finite, leave it
+				// nans will fail all compares so not go in here
+			}
+			else if ( f > FLT_MAX )
+			{
+				// +inf
+				SourceData[Idx] = FLT_MAX;
+			}
+			else if ( f < -FLT_MAX )
+			{
+				// -inf
+				SourceData[Idx] = -FLT_MAX;
+			}
+			else
+			{
+				// nan
+				SourceData[Idx] = 0.f;
+			}
+
+			//check( ! FMath::IsNaN( SourceData[Idx] ) );
+		}
+	}
 
 	Encode(
 		SourceData,
@@ -229,15 +266,18 @@ public:
 		EPixelFormat CompressedPixelFormat = GetEncodedPixelFormat(BuildSettings, bImageHasAlphaChannel);
 
 		bool bCompressionSucceeded = true;
-		int32 SliceSize = Image.SizeX * Image.SizeY;
+		int64 SliceSize = Image.GetSliceNumPixels();
 		for (int32 SliceIndex = 0; SliceIndex < Image.NumSlices && bCompressionSucceeded; ++SliceIndex)
 		{
 			TArray64<uint8> CompressedSliceData;
+			
+			const FLinearColor * SlicePixels = Image.AsRGBA32F().GetData() + SliceIndex * SliceSize;
 			bCompressionSucceeded = CompressImageUsingEtc2comp(
-				Image.AsRGBA32F().GetData() + SliceIndex * SliceSize,
+				const_cast<FLinearColor *>(SlicePixels),
 				CompressedPixelFormat,
 				Image.SizeX,
 				Image.SizeY,
+				SliceSize,
 				BuildSettings.GetDestGammaSpace(),
 				CompressedSliceData
 			);
